Changing image writing API from TArray to TArray64 data so images can be over 2B pixels.

diff --git a/Engine/Plugins/MovieScene/MovieRenderPipeline/Source/MovieRenderPipelineCore/Private/MovieRenderOverlappedImage.cpp b/Engine/Plugins/MovieScene/MovieRenderPipeline/Source/MovieRenderPipelineCore/Private/MovieRenderOverlappedImage.cpp
--- a/Engine/Plugins/MovieScene/MovieRenderPipeline/Source/MovieRenderPipelineCore/Private/MovieRenderOverlappedImage.cpp
+++ b/Engine/Plugins/MovieScene/MovieRenderPipeline/Source/MovieRenderPipelineCore/Private/MovieRenderOverlappedImage.cpp
@@ -429,7 +429,7 @@ void FImageOverlappedAccumulator::AccumulatePixelData(const FImagePixelData& InP
 	int32 RawBitDepth = InPixelData.GetBitDepth();
 	FIntPoint RawSize = InPixelData.GetSize();
 
-	int32 SizeInBytes = 0;
+	int64 SizeInBytes = 0;
 	const void* SrcRawDataPtr = nullptr;
 
 	bool IsFetchOk = InPixelData.GetRawData(SrcRawDataPtr, SizeInBytes);
@@ -563,10 +563,36 @@ void FImageOverlappedAccumulator::AccumulatePixelData(const FImagePixelData& InP
 			for (int ChanIter = 0; ChanIter < NumChannels; ChanIter++)
 			{
 				float* DstData = RawData[ChanIter].GetData();
+
 				int32 DstSize = RawSize.X * RawSize.Y;
-				for (int32 Iter = 0; Iter < DstSize; Iter++)
+
+				static bool EnableVectorizedPow = true;
+				if (EnableVectorizedPow)
+				{
+					int32 VectorGroups = DstSize / 4;
+					int32 VectorExtra = DstSize - 4 * VectorGroups;
+					
+					VectorRegister PowExponent = VectorSetFloat1(Gamma);
+
+					for (int32 Iter = 0; Iter < VectorGroups; Iter++)
+					{
+						VectorRegister Value = VectorLoad(&DstData[4*Iter]);
+						Value = VectorPow(Value, PowExponent);
+						VectorStore(Value, &DstData[4*Iter]);
+					}
+
+					for (int32 Iter = 0; Iter < VectorExtra; Iter++)
+					{
+						int32 Index = 4 * VectorGroups + Iter;
+						DstData[Index] = FMath::Pow(DstData[Index], Gamma);
+					}
+				}
+				else
 				{
-					DstData[Iter] = FMath::Pow(DstData[Iter], Gamma);
+					for (int32 Iter = 0; Iter < DstSize; Iter++)
+					{
+						DstData[Iter] = FMath::Pow(DstData[Iter], Gamma);
+					}
 				}
 			}
 		} 
@@ -665,7 +691,7 @@ void FImageOverlappedAccumulator::FetchFullImageValue(float Rgba[4], int32 FullX
 	}
 }
 
-void FImageOverlappedAccumulator::FetchFinalPixelDataByte(TArray<FColor> & OutPixelData) const
+void FImageOverlappedAccumulator::FetchFinalPixelDataByte(TArray64<FColor> & OutPixelData) const
 {
 	int32 FullSizeX = PlaneSize.X;
 	int32 FullSizeY = PlaneSize.Y;
@@ -687,7 +713,7 @@ void FImageOverlappedAccumulator::FetchFinalPixelDataByte(TArray<FColor> & OutPi
 	}
 }
 
-void FImageOverlappedAccumulator::FetchFinalPixelDataLinearColor(TArray<FLinearColor> & OutPixelData) const
+void FImageOverlappedAccumulator::FetchFinalPixelDataLinearColor(TArray64<FLinearColor> & OutPixelData) const
 {
 	int32 FullSizeX = PlaneSize.X;
 	int32 FullSizeY = PlaneSize.Y;
diff --git a/Engine/Plugins/MovieScene/MovieRenderPipeline/Source/MovieRenderPipelineCore/Private/MovieRenderTileImage.cpp b/Engine/Plugins/MovieScene/MovieRenderPipeline/Source/MovieRenderPipelineCore/Private/MovieRenderTileImage.cpp
--- a/Engine/Plugins/MovieScene/MovieRenderPipeline/Source/MovieRenderPipelineCore/Private/MovieRenderTileImage.cpp
+++ b/Engine/Plugins/MovieScene/MovieRenderPipeline/Source/MovieRenderPipelineCore/Private/MovieRenderTileImage.cpp
@@ -279,7 +279,7 @@ void FImageTileAccumulator::AccumulatePixelData(const FImagePixelData& InPixelDa
 	int32 RawBitDepth = InPixelData.GetBitDepth();
 	FIntPoint RawSize = InPixelData.GetSize();
 
-	int32 SizeInBytes = 0;
+	int64 SizeInBytes = 0;
 	const void* SrcRawDataPtr = nullptr;
 
 	bool IsFetchOk = InPixelData.GetRawData(SrcRawDataPtr, SizeInBytes);
@@ -483,7 +483,7 @@ void FImageTileAccumulator::FetchFinalPlaneScale(TArray<float>& PlaneScale) cons
 	}
 }
 
-void FImageTileAccumulator::FetchFinalPixelDataByte(TArray<FColor> & OutPixelData) const
+void FImageTileAccumulator::FetchFinalPixelDataByte(TArray64<FColor> & OutPixelData) const
 {
 	int32 FullSizeX = TileSizeX * NumTilesX;
 	int32 FullSizeY = TileSizeY * NumTilesY;
@@ -508,7 +508,7 @@ void FImageTileAccumulator::FetchFinalPixelDataByte(TArray<FColor> & OutPixelDat
 	}
 }
 
-void FImageTileAccumulator::FetchFinalPixelDataLinearColor(TArray<FLinearColor> & OutPixelData) const
+void FImageTileAccumulator::FetchFinalPixelDataLinearColor(TArray64<FLinearColor> & OutPixelData) const
 {
 	int32 FullSizeX = TileSizeX * NumTilesX;
 	int32 FullSizeY = TileSizeY * NumTilesY;
diff --git a/Engine/Plugins/MovieScene/MovieRenderPipeline/Source/MovieRenderPipelineCore/Public/MovieRenderOverlappedImage.h b/Engine/Plugins/MovieScene/MovieRenderPipeline/Source/MovieRenderPipelineCore/Public/MovieRenderOverlappedImage.h
--- a/Engine/Plugins/MovieScene/MovieRenderPipeline/Source/MovieRenderPipelineCore/Public/MovieRenderOverlappedImage.h
+++ b/Engine/Plugins/MovieScene/MovieRenderPipeline/Source/MovieRenderPipelineCore/Public/MovieRenderOverlappedImage.h
@@ -151,14 +151,14 @@ public:
 	 * 
 	 * @param FImagePixelData - Finished pixel data.
 	 */
-	void FetchFinalPixelDataByte(TArray<FColor>& OutPixelData) const;
+	void FetchFinalPixelDataByte(TArray64<FColor>& OutPixelData) const;
 
 	/**
 	 * After accumulation is finished, fetch the final image as linear colors
 	 * 
 	 * @param FImagePixelData - Finished pixel data.
 	 */
-	void FetchFinalPixelDataLinearColor(TArray<FLinearColor>& OutPixelData) const;
+	void FetchFinalPixelDataLinearColor(TArray64<FLinearColor>& OutPixelData) const;
 
 	/**
 	 * Grab a single pixel from the full res tile and scale it by the appropriate Scale value.
diff --git a/Engine/Plugins/MovieScene/MovieRenderPipeline/Source/MovieRenderPipelineCore/Public/MovieRenderTileImage.h b/Engine/Plugins/MovieScene/MovieRenderPipeline/Source/MovieRenderPipelineCore/Public/MovieRenderTileImage.h
--- a/Engine/Plugins/MovieScene/MovieRenderPipeline/Source/MovieRenderPipelineCore/Public/MovieRenderTileImage.h
+++ b/Engine/Plugins/MovieScene/MovieRenderPipeline/Source/MovieRenderPipelineCore/Public/MovieRenderTileImage.h
@@ -178,14 +178,14 @@ public:
 	 * 
 	 * @param FImagePixelData - Finished pixel data.
 	 */
-	void FetchFinalPixelDataByte(TArray<FColor>& OutPixelData) const;
+	void FetchFinalPixelDataByte(TArray64<FColor>& OutPixelData) const;
 
 	/**
 	 * After accumulation is finished, fetch the final image as linear colors
 	 * 
 	 * @param FImagePixelData - Finished pixel data.
 	 */
-	void FetchFinalPixelDataLinearColor(TArray<FLinearColor>& OutPixelData) const;
+	void FetchFinalPixelDataLinearColor(TArray64<FLinearColor>& OutPixelData) const;
 
 	/**
 	 * After accumulation is finished, fetch the final scale of each plane.
diff --git a/Engine/Plugins/MovieScene/MovieRenderPipeline/Source/MovieRenderPipelineRenderPasses/Private/MoviePipelineBackbufferPass.cpp b/Engine/Plugins/MovieScene/MovieRenderPipeline/Source/MovieRenderPipelineRenderPasses/Private/MoviePipelineBackbufferPass.cpp
--- a/Engine/Plugins/MovieScene/MovieRenderPipeline/Source/MovieRenderPipelineRenderPasses/Private/MoviePipelineBackbufferPass.cpp
+++ b/Engine/Plugins/MovieScene/MovieRenderPipeline/Source/MovieRenderPipelineRenderPasses/Private/MoviePipelineBackbufferPass.cpp
@@ -345,8 +345,11 @@ void ReadAndAccumulateSample_RenderThread(FRHICommandListImmediate &RHICmdList,
 	FrameData->OverlappedSubpixelShift = InParams.PassMetrics.OverlappedSubpixelShift;
 
 	// Read the data back to the CPU
-	TUniquePtr<TImagePixelData<FColor>> PixelData = MakeUnique<TImagePixelData<FColor>>(SourceRect.Size(), FrameData);
-	RHICmdList.ReadSurfaceData(InParams.CanvasRenderTarget->GetRenderTargetTexture(), SourceRect, PixelData->Pixels, ReadDataFlags);
+	TArray<FColor> RawPixels;
+	RawPixels.SetNum(SourceRect.Width() * SourceRect.Height());
+
+	RHICmdList.ReadSurfaceData(InParams.CanvasRenderTarget->GetRenderTargetTexture(), SourceRect, RawPixels, ReadDataFlags);
+	TUniquePtr<TImagePixelData<FColor>> PixelData = MakeUnique<TImagePixelData<FColor>>(SourceRect.Size(), TArray64<FColor>(MoveTemp(RawPixels)), FrameData);
 
 	check(PixelData->IsDataWellFormed());
 
diff --git a/Engine/Source/Editor/UnrealEd/Private/EditorViewportClient.cpp b/Engine/Source/Editor/UnrealEd/Private/EditorViewportClient.cpp
--- a/Engine/Source/Editor/UnrealEd/Private/EditorViewportClient.cpp
+++ b/Engine/Source/Editor/UnrealEd/Private/EditorViewportClient.cpp
@@ -5332,14 +5332,17 @@ void FEditorViewportClient::TakeScreenshot(FViewport* InViewport, bool bInValida
 
 	{
 		// Read the contents of the viewport into an array.
-		TUniquePtr<TImagePixelData<FColor>> PixelData = MakeUnique<TImagePixelData<FColor>>(InViewport->GetSizeXY());
-		if( !InViewport->ReadPixels(PixelData->Pixels) )
+		TArray<FColor> RawPixels;
+		RawPixels.SetNum(InViewport->GetSizeXY().X * InViewport->GetSizeXY().Y);
+		if( !InViewport->ReadPixels(RawPixels) )
 		{
 			// Failed to read the image from the viewport
 			SaveMessagePtr->SetText(NSLOCTEXT( "UnrealEd", "ScreenshotFailedViewport", "Screenshot failed, unable to read image from viewport" ));
 			return;
 		}
 
+		TUniquePtr<TImagePixelData<FColor>> PixelData = MakeUnique<TImagePixelData<FColor>>(InViewport->GetSizeXY(), TArray64<FColor>(MoveTemp(RawPixels)));
+
 		check(PixelData->IsDataWellFormed());
 		ImageTask->PixelData = MoveTemp(PixelData);
 	}
@@ -5546,7 +5549,7 @@ bool FEditorViewportClient::ProcessScreenShots(FViewport* InViewport)
 			}
 
 			TUniquePtr<FImageWriteTask> ImageTask = MakeUnique<FImageWriteTask>();
-			ImageTask->PixelData = MakeUnique<TImagePixelData<FColor>>(BitmapSize, MoveTemp(Bitmap));
+			ImageTask->PixelData = MakeUnique<TImagePixelData<FColor>>(BitmapSize, TArray64<FColor>(MoveTemp(Bitmap)));
 
 			// Set full alpha on the bitmap
 			if (!bWriteAlpha)
diff --git a/Engine/Source/Runtime/ImageWriteQueue/Private/ImageWriteBlueprintLibrary.cpp b/Engine/Source/Runtime/ImageWriteQueue/Private/ImageWriteBlueprintLibrary.cpp
--- a/Engine/Source/Runtime/ImageWriteQueue/Private/ImageWriteBlueprintLibrary.cpp
+++ b/Engine/Source/Runtime/ImageWriteQueue/Private/ImageWriteBlueprintLibrary.cpp
@@ -22,6 +22,7 @@ EImageFormat ImageFormatFromDesired(EDesiredImageFormat In)
 	return EImageFormat::BMP;
 }
 
+
 bool UImageWriteBlueprintLibrary::ResolvePixelData(UTexture* InTexture, const FOnPixelsReady& OnPixelsReady)
 {
 	if (!InTexture)
@@ -70,9 +71,12 @@ bool UImageWriteBlueprintLibrary::ResolvePixelData(UTexture* InTexture, const FO
 			{
 				case PF_FloatRGBA:
 				{
-					TUniquePtr<TImagePixelData<FFloat16Color>> PixelData = MakeUnique<TImagePixelData<FFloat16Color>>(SourceRect.Size());
+					TArray<FFloat16Color> RawPixels;
+					RawPixels.SetNum(SourceRect.Width() * SourceRect.Height());
+					RHICmdList.ReadSurfaceFloatData(Texture2D, SourceRect, RawPixels, (ECubeFace)0, 0, 0);
+
+					TUniquePtr<TImagePixelData<FFloat16Color>> PixelData = MakeUnique<TImagePixelData<FFloat16Color>>(SourceRect.Size(), TArray64<FFloat16Color>(MoveTemp(RawPixels)));
 
-					RHICmdList.ReadSurfaceFloatData(Texture2D, SourceRect, PixelData->Pixels, (ECubeFace)0, 0, 0);
 					if (PixelData->IsDataWellFormed())
 					{
 						OnPixelsReady(MoveTemp(PixelData));
@@ -87,9 +91,11 @@ bool UImageWriteBlueprintLibrary::ResolvePixelData(UTexture* InTexture, const FO
 					FReadSurfaceDataFlags ReadDataFlags(RCM_MinMax);
 					ReadDataFlags.SetLinearToGamma(false);
 
-					TUniquePtr<TImagePixelData<FLinearColor>> PixelData = MakeUnique<TImagePixelData<FLinearColor>>(SourceRect.Size());
 
-					RHICmdList.ReadSurfaceData(Texture2D, SourceRect, PixelData->Pixels, ReadDataFlags);
+					TArray<FLinearColor> RawPixels;
+					RawPixels.SetNum(SourceRect.Width() * SourceRect.Height());
+					RHICmdList.ReadSurfaceData(Texture2D, SourceRect, RawPixels, ReadDataFlags);
+					TUniquePtr<TImagePixelData<FLinearColor>> PixelData = MakeUnique<TImagePixelData<FLinearColor>>(SourceRect.Size(), TArray64<FLinearColor>(MoveTemp(RawPixels)));
 					if (PixelData->IsDataWellFormed())
 					{
 						OnPixelsReady(MoveTemp(PixelData));
@@ -105,9 +111,11 @@ bool UImageWriteBlueprintLibrary::ResolvePixelData(UTexture* InTexture, const FO
 					FReadSurfaceDataFlags ReadDataFlags;
 					ReadDataFlags.SetLinearToGamma(false);
 
-					TUniquePtr<TImagePixelData<FColor>> PixelData = MakeUnique<TImagePixelData<FColor>>(SourceRect.Size());
+					TArray<FColor> RawPixels;
+					RawPixels.SetNum(SourceRect.Width() * SourceRect.Height());
+					RHICmdList.ReadSurfaceData(Texture2D, SourceRect, RawPixels, ReadDataFlags);
 
-					RHICmdList.ReadSurfaceData(Texture2D, SourceRect, PixelData->Pixels, ReadDataFlags);
+					TUniquePtr<TImagePixelData<FColor>> PixelData = MakeUnique<TImagePixelData<FColor>>(SourceRect.Size(), TArray64<FColor>(MoveTemp(RawPixels)));
 					if (PixelData->IsDataWellFormed())
 					{
 						OnPixelsReady(MoveTemp(PixelData));
diff --git a/Engine/Source/Runtime/ImageWriteQueue/Private/ImageWriteTask.cpp b/Engine/Source/Runtime/ImageWriteQueue/Private/ImageWriteTask.cpp
--- a/Engine/Source/Runtime/ImageWriteQueue/Private/ImageWriteTask.cpp
+++ b/Engine/Source/Runtime/ImageWriteQueue/Private/ImageWriteTask.cpp
@@ -109,7 +109,7 @@ void FImageWriteTask::OnAbandoned()
 bool FImageWriteTask::InitializeWrapper(IImageWrapper* InWrapper, EImageFormat WrapperFormat)
 {
 	const void* RawPtr = nullptr;
-	int32 SizeBytes = 0;
+	int64 SizeBytes = 0;
 
 	if (PixelData->GetRawData(RawPtr, SizeBytes))
 	{
@@ -135,7 +135,7 @@ bool FImageWriteTask::WriteBitmap()
 	}
 
 	const void* RawPtr = nullptr;
-	int32 SizeBytes = 0;
+	int64 SizeBytes = 0;
 
 	if (PixelData->GetRawData(RawPtr, SizeBytes))
 	{
diff --git a/Engine/Source/Runtime/ImageWriteQueue/Public/ImagePixelData.h b/Engine/Source/Runtime/ImageWriteQueue/Public/ImagePixelData.h
--- a/Engine/Source/Runtime/ImageWriteQueue/Public/ImagePixelData.h
+++ b/Engine/Source/Runtime/ImageWriteQueue/Public/ImagePixelData.h
@@ -75,7 +75,7 @@ struct FImagePixelData
 	bool IsDataWellFormed() const
 	{
 		const void* RawPtr    = nullptr;
-		int32       SizeBytes = 0;
+		int64       SizeBytes = 0;
 
 		return GetRawData(RawPtr, SizeBytes);
 	}
@@ -83,10 +83,10 @@ struct FImagePixelData
 	/**
 	 * Get the data and its size only if it is well formed
 	 */
-	bool GetRawData(const void*& OutRawData, int32& OutSizeBytes) const
+	bool GetRawData(const void*& OutRawData, int64& OutSizeBytes) const
 	{
 		const void* RawPtr    = nullptr;
-		int32       SizeBytes = 0;
+		int64       SizeBytes = 0;
 
 		RetrieveData(RawPtr, SizeBytes);
 
@@ -138,7 +138,7 @@ private:
 	/**
 	 * Retrieve the raw pixel data
 	 */
-	virtual void RetrieveData(const void*& OutDataPtr, int32& OutSizeBytes) const = 0;
+	virtual void RetrieveData(const void*& OutDataPtr, int64& OutSizeBytes) const = 0;
 	virtual TUniquePtr<FImagePixelData> Copy() const = 0;
 	virtual TUniquePtr<FImagePixelData> Move() = 0;
 
@@ -167,13 +167,13 @@ private:
 template<typename PixelType>
 struct TImagePixelData : FImagePixelData
 {
-	TArray<PixelType> Pixels;
+	TArray64<PixelType> Pixels;
 
 	TImagePixelData(const FIntPoint& InSize)
 		: FImagePixelData(InSize, TImagePixelDataTraits<PixelType>::PixelType, TImagePixelDataTraits<PixelType>::PixelLayout, TImagePixelDataTraits<PixelType>::BitDepth, TImagePixelDataTraits<PixelType>::NumChannels, nullptr)
 	{}
 
-	TImagePixelData(const FIntPoint& InSize, TArray<PixelType>&& InPixels)
+	TImagePixelData(const FIntPoint& InSize, TArray64<PixelType>&& InPixels)
 		: FImagePixelData(InSize, TImagePixelDataTraits<PixelType>::PixelType, TImagePixelDataTraits<PixelType>::PixelLayout, TImagePixelDataTraits<PixelType>::BitDepth, TImagePixelDataTraits<PixelType>::NumChannels, nullptr)
 		, Pixels(MoveTemp(InPixels))
 	{}
@@ -182,7 +182,7 @@ struct TImagePixelData : FImagePixelData
 		: FImagePixelData(InSize, TImagePixelDataTraits<PixelType>::PixelType, TImagePixelDataTraits<PixelType>::PixelLayout, TImagePixelDataTraits<PixelType>::BitDepth, TImagePixelDataTraits<PixelType>::NumChannels, InPayload)
 	{}
 
-	TImagePixelData(const FIntPoint& InSize, TArray<PixelType>&& InPixels, FImagePixelPayloadPtr InPayload)
+	TImagePixelData(const FIntPoint& InSize, TArray64<PixelType>&& InPixels, FImagePixelPayloadPtr InPayload)
 		: FImagePixelData(InSize, TImagePixelDataTraits<PixelType>::PixelType, TImagePixelDataTraits<PixelType>::PixelLayout, TImagePixelDataTraits<PixelType>::BitDepth, TImagePixelDataTraits<PixelType>::NumChannels, InPayload)
 		, Pixels(MoveTemp(InPixels))
 	{}
@@ -197,7 +197,7 @@ struct TImagePixelData : FImagePixelData
 		return MakeUnique<TImagePixelData<PixelType>>(*this);
 	}
 
-	virtual void RetrieveData(const void*& OutDataPtr, int32& OutSizeBytes) const override
+	virtual void RetrieveData(const void*& OutDataPtr, int64& OutSizeBytes) const override
 	{
 		OutDataPtr = static_cast<const void*>(&Pixels[0]);
 		OutSizeBytes = Pixels.Num() * sizeof(PixelType);
diff --git a/Engine/Source/Runtime/MediaUtils/Private/MediaRecorder.cpp b/Engine/Source/Runtime/MediaUtils/Private/MediaRecorder.cpp
--- a/Engine/Source/Runtime/MediaUtils/Private/MediaRecorder.cpp
+++ b/Engine/Source/Runtime/MediaUtils/Private/MediaRecorder.cpp
@@ -79,10 +79,10 @@ struct FMediaImagePixelData : FImagePixelData
 		return MakeUnique<FMediaImagePixelData>(*this);
 	}
 
-	virtual void RetrieveData(const void*& OutDataPtr, int32& OutSizeBytes) const override
+	virtual void RetrieveData(const void*& OutDataPtr, int64& OutSizeBytes) const override
 	{
 		OutDataPtr = Sample->GetBuffer();
-		OutSizeBytes = Sample->GetStride() * Sample->GetDim().Y;
+		OutSizeBytes = int64(Sample->GetStride()) * int64(Sample->GetDim().Y);
 	}
 };
 
@@ -95,19 +95,19 @@ namespace MediaRecorderHelpers
 	template<class TColorType>
 	TUniquePtr<TImagePixelData<TColorType>> CreatePixelData(TSharedPtr<IMediaTextureSample, ESPMode::ThreadSafe> InSample, const FIntPoint InSize, int32 InNumChannels)
 	{
-		const int32 NumberOfTexel = InSize.Y * InSize.X;
+		const int64 NumberOfTexel = InSize.Y * InSize.X;
 		TUniquePtr<TImagePixelData<TColorType>> PixelData = MakeUnique<TImagePixelData<TColorType>>(InSize);
 		PixelData->Pixels.Reset(NumberOfTexel);
 
 		const void* Buffer = InSample->GetBuffer();
-		const uint32 Stride = InSample->GetStride();
-		if (NumberOfTexel == Stride *InSize.Y / InNumChannels)
+		const uint64 Stride = InSample->GetStride();
+		if (NumberOfTexel == (Stride *InSize.Y) / InNumChannels)
 		{
 			PixelData->Pixels.Append(reinterpret_cast<const TColorType*>(Buffer), NumberOfTexel);
 		}
 		else
 		{
-			for (int IndexY = 0; IndexY < InSize.Y; ++IndexY)
+			for (int64 IndexY = 0; IndexY < InSize.Y; ++IndexY)
 			{
 				PixelData->Pixels.Append(reinterpret_cast<const TColorType*>(reinterpret_cast<const uint8*>(Buffer) + (Stride*IndexY)), InSize.X);
 			}
@@ -118,12 +118,12 @@ namespace MediaRecorderHelpers
 
 	TUniquePtr<TImagePixelData<FColor>> ConvertSampleInterpretedAsFloat16ColorToFColor(TSharedPtr<IMediaTextureSample, ESPMode::ThreadSafe> InSample, const FIntPoint InSize, int32 InNumChannels)
 	{
-		const int32 NumberOfTexel = InSize.Y * InSize.X;
+		const int64 NumberOfTexel = InSize.Y * InSize.X;
 		TUniquePtr<TImagePixelData<FColor>> PixelData = MakeUnique<TImagePixelData<FColor>>(InSize);
 		PixelData->Pixels.Reset(NumberOfTexel);
 
 		const FFloat16Color* Buffer = reinterpret_cast<const FFloat16Color*>(InSample->GetBuffer());
-		for (int32 i = 0 ; i < NumberOfTexel; i++)
+		for (int64 i = 0 ; i < NumberOfTexel; i++)
 		{
 			FColor Output;
 			Output.R = (uint8)(FMath::Clamp<float>(Buffer[0].R, 0.0f, 1.0f) * 255);
diff --git a/Engine/Source/Runtime/MovieSceneCapture/Private/ImageSequenceProtocol.cpp b/Engine/Source/Runtime/MovieSceneCapture/Private/ImageSequenceProtocol.cpp
--- a/Engine/Source/Runtime/MovieSceneCapture/Private/ImageSequenceProtocol.cpp
+++ b/Engine/Source/Runtime/MovieSceneCapture/Private/ImageSequenceProtocol.cpp
@@ -122,7 +122,7 @@ void UImageSequenceProtocol::ProcessFrame(FCapturedFrameData Frame)
 	TUniquePtr<FImageWriteTask> ImageTask = MakeUnique<FImageWriteTask>();
 
 	// Move the color buffer into a raw image data container that we can pass to the write queue
-	ImageTask->PixelData = MakeUnique<TImagePixelData<FColor>>(Frame.BufferSize, MoveTemp(Frame.ColorBuffer));
+	ImageTask->PixelData = MakeUnique<TImagePixelData<FColor>>(Frame.BufferSize, TArray64<FColor>(MoveTemp(Frame.ColorBuffer)));
 	if (Format == EImageFormat::PNG)
 	{
 		// Always write full alpha for PNGs
diff --git a/Engine/Source/Runtime/MovieSceneCapture/Private/UserDefinedCaptureProtocol.cpp b/Engine/Source/Runtime/MovieSceneCapture/Private/UserDefinedCaptureProtocol.cpp
--- a/Engine/Source/Runtime/MovieSceneCapture/Private/UserDefinedCaptureProtocol.cpp
+++ b/Engine/Source/Runtime/MovieSceneCapture/Private/UserDefinedCaptureProtocol.cpp
@@ -115,7 +115,7 @@ void UUserDefinedCaptureProtocol::TickImpl()
 		for (FCapturedFrameData& Frame : CapturedFrames)
 		{
 			// Steal the frame and make it shareable
-			FCapturedPixels CapturedPixels { MakeShared<TImagePixelData<FColor>, ESPMode::ThreadSafe>( Frame.BufferSize, MoveTemp(Frame.ColorBuffer) ) };
+			FCapturedPixels CapturedPixels { MakeShared<TImagePixelData<FColor>, ESPMode::ThreadSafe>( Frame.BufferSize, TArray64<FColor>(MoveTemp(Frame.ColorBuffer) ) ) };
 			FFrameMetrics   CapturedMetrics = static_cast<FCaptureProtocolFrameData*>(Frame.Payload.Get())->Metrics;
 
 			// Call the handler
diff --git a/Engine/Source/Runtime/Renderer/Private/PostProcess/PostProcessVisualizeBuffer.cpp b/Engine/Source/Runtime/Renderer/Private/PostProcess/PostProcessVisualizeBuffer.cpp
--- a/Engine/Source/Runtime/Renderer/Private/PostProcess/PostProcessVisualizeBuffer.cpp
+++ b/Engine/Source/Runtime/Renderer/Private/PostProcess/PostProcessVisualizeBuffer.cpp
@@ -193,8 +193,10 @@ TUniquePtr<FImagePixelData> ReadbackPixelData(FRHICommandListImmediate& RHICmdLi
 	{
 	case PF_FloatRGBA:
 	{
-		TUniquePtr<TImagePixelData<FFloat16Color>> PixelData = MakeUnique<TImagePixelData<FFloat16Color>>(SourceRect.Size());
-		RHICmdList.ReadSurfaceFloatData(Texture, SourceRect, PixelData->Pixels, (ECubeFace)0, 0, 0);
+		TArray<FFloat16Color> RawPixels;
+		RawPixels.SetNum(SourceRect.Width() * SourceRect.Height());
+		RHICmdList.ReadSurfaceFloatData(Texture, SourceRect, RawPixels, (ECubeFace)0, 0, 0);
+		TUniquePtr<TImagePixelData<FFloat16Color>> PixelData = MakeUnique<TImagePixelData<FFloat16Color>>(SourceRect.Size(), TArray64<FFloat16Color>(MoveTemp(RawPixels)));
 
 		check(PixelData->IsDataWellFormed());
 		return PixelData;
@@ -205,8 +207,10 @@ TUniquePtr<FImagePixelData> ReadbackPixelData(FRHICommandListImmediate& RHICmdLi
 		FReadSurfaceDataFlags ReadDataFlags(RCM_MinMax);
 		ReadDataFlags.SetLinearToGamma(false);
 
-		TUniquePtr<TImagePixelData<FLinearColor>> PixelData = MakeUnique<TImagePixelData<FLinearColor>>(SourceRect.Size());
-		RHICmdList.ReadSurfaceData(Texture, SourceRect, PixelData->Pixels, ReadDataFlags);
+		TArray<FLinearColor> RawPixels;
+		RawPixels.SetNum(SourceRect.Width() * SourceRect.Height());
+		RHICmdList.ReadSurfaceData(Texture, SourceRect, RawPixels, ReadDataFlags);
+		TUniquePtr<TImagePixelData<FLinearColor>> PixelData = MakeUnique<TImagePixelData<FLinearColor>>(SourceRect.Size(), TArray64<FLinearColor>(MoveTemp(RawPixels)));
 
 		check(PixelData->IsDataWellFormed());
 		return PixelData;
@@ -218,8 +222,10 @@ TUniquePtr<FImagePixelData> ReadbackPixelData(FRHICommandListImmediate& RHICmdLi
 		FReadSurfaceDataFlags ReadDataFlags;
 		ReadDataFlags.SetLinearToGamma(false);
 
-		TUniquePtr<TImagePixelData<FColor>> PixelData = MakeUnique<TImagePixelData<FColor>>(SourceRect.Size());
-		RHICmdList.ReadSurfaceData(Texture, SourceRect, PixelData->Pixels, ReadDataFlags);
+		TArray<FColor> RawPixels;
+		RawPixels.SetNum(SourceRect.Width() * SourceRect.Height());
+		RHICmdList.ReadSurfaceData(Texture, SourceRect, RawPixels, ReadDataFlags);
+		TUniquePtr<TImagePixelData<FColor>> PixelData = MakeUnique<TImagePixelData<FColor>>(SourceRect.Size(), TArray64<FColor>(MoveTemp(RawPixels)));
 
 		check(PixelData->IsDataWellFormed());
 		return PixelData;
diff --git a/Engine/Source/Runtime/Renderer/Private/PostProcess/RenderingCompositionGraph.cpp b/Engine/Source/Runtime/Renderer/Private/PostProcess/RenderingCompositionGraph.cpp
--- a/Engine/Source/Runtime/Renderer/Private/PostProcess/RenderingCompositionGraph.cpp
+++ b/Engine/Source/Runtime/Renderer/Private/PostProcess/RenderingCompositionGraph.cpp
@@ -407,8 +407,10 @@ TUniquePtr<FImagePixelData> FRenderingCompositionGraph::GetDumpOutput(FRendering
 	{
 		case PF_FloatRGBA:
 		{
-			TUniquePtr<TImagePixelData<FFloat16Color>> PixelData = MakeUnique<TImagePixelData<FFloat16Color>>(SourceRect.Size());
-			Context.RHICmdList.ReadSurfaceFloatData(Texture, SourceRect, PixelData->Pixels, (ECubeFace)0, 0, 0);
+			TArray<FFloat16Color> RawPixels;
+			RawPixels.SetNum(SourceRect.Width() * SourceRect.Height());
+			Context.RHICmdList.ReadSurfaceFloatData(Texture, SourceRect, RawPixels, (ECubeFace)0, 0, 0);
+			TUniquePtr<TImagePixelData<FFloat16Color>> PixelData = MakeUnique<TImagePixelData<FFloat16Color>>(SourceRect.Size(), TArray64<FFloat16Color>(MoveTemp(RawPixels)));
 
 			check(PixelData->IsDataWellFormed());
 
@@ -420,8 +422,10 @@ TUniquePtr<FImagePixelData> FRenderingCompositionGraph::GetDumpOutput(FRendering
 			FReadSurfaceDataFlags ReadDataFlags(RCM_MinMax);
 			ReadDataFlags.SetLinearToGamma(false);
 
-			TUniquePtr<TImagePixelData<FLinearColor>> PixelData = MakeUnique<TImagePixelData<FLinearColor>>(SourceRect.Size());
-			Context.RHICmdList.ReadSurfaceData(Texture, SourceRect, PixelData->Pixels, ReadDataFlags);
+			TArray<FLinearColor> RawPixels;
+			RawPixels.SetNum(SourceRect.Width() * SourceRect.Height());
+			Context.RHICmdList.ReadSurfaceData(Texture, SourceRect, RawPixels, ReadDataFlags);
+			TUniquePtr<TImagePixelData<FLinearColor>> PixelData = MakeUnique<TImagePixelData<FLinearColor>>(SourceRect.Size(), TArray64<FLinearColor>(MoveTemp(RawPixels)));
 
 			check(PixelData->IsDataWellFormed());
 
@@ -434,8 +438,10 @@ TUniquePtr<FImagePixelData> FRenderingCompositionGraph::GetDumpOutput(FRendering
 			FReadSurfaceDataFlags ReadDataFlags;
 			ReadDataFlags.SetLinearToGamma(false);
 
-			TUniquePtr<TImagePixelData<FColor>> PixelData = MakeUnique<TImagePixelData<FColor>>(SourceRect.Size());
-			Context.RHICmdList.ReadSurfaceData(Texture, SourceRect, PixelData->Pixels, ReadDataFlags);
+			TArray<FColor> RawPixels;
+			RawPixels.SetNum(SourceRect.Width() * SourceRect.Height());
+			Context.RHICmdList.ReadSurfaceData(Texture, SourceRect, RawPixels, ReadDataFlags);
+			TUniquePtr<TImagePixelData<FColor>> PixelData = MakeUnique<TImagePixelData<FColor>>(SourceRect.Size(), TArray64<FColor>(MoveTemp(RawPixels)));
 
 			check(PixelData->IsDataWellFormed());
 
