Better fix for: Avoid spawning 6 extra threads on mobile for saving HighResScreenshots

diff --git a/Engine/Source/Runtime/ImageWriteQueue/Private/ImageWriteQueue.cpp b/Engine/Source/Runtime/ImageWriteQueue/Private/ImageWriteQueue.cpp
--- a/Engine/Source/Runtime/ImageWriteQueue/Private/ImageWriteQueue.cpp
+++ b/Engine/Source/Runtime/ImageWriteQueue/Private/ImageWriteQueue.cpp
@@ -243,8 +243,14 @@ void FImageWriteQueue::RecreateThreadPool()
 	// Prevent any other tasks being dispatched
 	FScopeLock ScopeLock(&ThreadPoolMutex);
 
+#if (PLATFORM_IOS || PLATFORM_TVOS || PLATFORM_ANDROID)
+	// To avoid spawning extra threads use global IO thread on mobile 
+	const int32 MaxConcurrency = GIOThreadPool->GetNumThreads();
+#else
 	const int32 ConfiguredMaxConcurrency = CVarImageWriteQueueMaxConcurrency.GetValueOnAnyThread();
 	const int32 MaxConcurrency = ConfiguredMaxConcurrency == -1 ? FPlatformMisc::NumberOfCores() : ConfiguredMaxConcurrency;
+#endif
+
 	if (ThreadPool && MaxConcurrency != ThreadPool->GetNumThreads())
 	{
 		CreateFence().Wait();
diff --git a/Engine/Source/Runtime/Launch/Private/LaunchEngineLoop.cpp b/Engine/Source/Runtime/Launch/Private/LaunchEngineLoop.cpp
--- a/Engine/Source/Runtime/Launch/Private/LaunchEngineLoop.cpp
+++ b/Engine/Source/Runtime/Launch/Private/LaunchEngineLoop.cpp
@@ -3516,13 +3516,11 @@ int32 FEngineLoop::PreInitPostStartupScreen(const TCHAR* CmdLine)
 	}
 #endif
 
-#if !(PLATFORM_IOS || PLATFORM_TVOS || PLATFORM_ANDROID)
 	// Init HighRes screenshot system, unless running on server
 	if (!IsRunningDedicatedServer())
 	{
 		GetHighResScreenshotConfig().Init();
 	}
-#endif
 
 #else // WITH_ENGINE
 	InitEngineTextLocalization();
