Reuse ThreadPool for ImageWriteQueue to get rid of another set of threads

diff --git a/Engine/Source/Runtime/ImageWriteQueue/Private/ImageWriteQueue.cpp b/Engine/Source/Runtime/ImageWriteQueue/Private/ImageWriteQueue.cpp
--- a/Engine/Source/Runtime/ImageWriteQueue/Private/ImageWriteQueue.cpp
+++ b/Engine/Source/Runtime/ImageWriteQueue/Private/ImageWriteQueue.cpp
@@ -10,6 +10,7 @@
 #include "Async/Async.h"
 #include "Modules/ModuleManager.h"
 #include "IImageWrapperModule.h"
+#include "Misc/QueuedThreadPoolWrapper.h"
 
 DEFINE_LOG_CATEGORY(LogImageWriteQueue);
 
@@ -276,6 +277,12 @@ void FImageWriteQueue::RecreateThreadPool()
 			bOwnedThreadPool = false;
 			ThreadPool = GIOThreadPool;
 		}
+		else if (GThreadPool && GThreadPool->GetNumThreads() >= MaxConcurrency)
+		{
+			// Use a simple wrapper to limit concurrency and reuse threads we already have
+			bOwnedThreadPool = true;
+			ThreadPool = new FQueuedThreadPoolWrapper(GThreadPool, MaxConcurrency);
+		}
 		else
 		{
 			// Create a new thread pool as a last resort
