Sequencer: Fix signess issue. GetSizeX and GetSizeY() are unsigned integers.

diff --git a/Engine/Source/Runtime/ImageWriteQueue/Private/ImageWriteBlueprintLibrary.cpp b/Engine/Source/Runtime/ImageWriteQueue/Private/ImageWriteBlueprintLibrary.cpp
--- a/Engine/Source/Runtime/ImageWriteQueue/Private/ImageWriteBlueprintLibrary.cpp
+++ b/Engine/Source/Runtime/ImageWriteQueue/Private/ImageWriteBlueprintLibrary.cpp
@@ -68,7 +68,7 @@ bool UImageWriteBlueprintLibrary::ResolvePixelData(UTexture* InTexture, const FO
 				return;
 			}
 
-			FIntRect SourceRect(0, 0, Texture2D->GetSizeX(), Texture2D->GetSizeY());
+			FIntRect SourceRect(0, 0, Texture2D->GetDesc().Extent.X, Texture2D->GetDesc().Extent.Y);
 			switch (Texture2D->GetFormat())
 			{
 				case PF_FloatRGBA:
