Fix redundant file writes and overwrites when saving high-res screenshots. GIsHighResScreenshot was being set to true too early, so images would be repeatedly generated, saved, and overwritten. Sometimes ImageWriteTask workers would try to delete a file while it was still being written, giving an error that was kind of vague about the actual cause. Setting GIsHighResScreenshot to true on the last frame of the frame delay only schedules ImageWriteTasks for the final frame so redundant images aren't written. Also improve error message when overwriting a file fails.

#rnx

diff --git a/Engine/Source/Runtime/Engine/Private/UnrealClient.cpp b/Engine/Source/Runtime/Engine/Private/UnrealClient.cpp
--- a/Engine/Source/Runtime/Engine/Private/UnrealClient.cpp
+++ b/Engine/Source/Runtime/Engine/Private/UnrealClient.cpp
@@ -1572,12 +1572,15 @@ void FViewport::HighResScreenshot()
 	// End the frame that was started before HighResScreenshot() was called. Pass nullptr for the viewport because there's no need to
 	// call EndRenderFrame on the real viewport, as BeginRenderFrame hasn't been called yet.
 	HighResScreenshotEndFrame(nullptr);
-
+	GIsHighResScreenshot = false;
 	// Perform run-up.
 	while (FrameDelay)
 	{
 		HighResScreenshotBeginFrame(DummyViewport);
-
+		if (FrameDelay == 1)
+		{
+			GIsHighResScreenshot = true;
+		}
 		FCanvas Canvas(DummyViewport, NULL, ViewportClient->GetWorld(), ViewportClient->GetWorld()->GetFeatureLevel());
 		{
 			ViewportClient->Draw(DummyViewport, &Canvas);
diff --git a/Engine/Source/Runtime/ImageWriteQueue/Private/ImageWriteTask.cpp b/Engine/Source/Runtime/ImageWriteQueue/Private/ImageWriteTask.cpp
--- a/Engine/Source/Runtime/ImageWriteQueue/Private/ImageWriteTask.cpp
+++ b/Engine/Source/Runtime/ImageWriteQueue/Private/ImageWriteTask.cpp
@@ -130,9 +130,14 @@ bool FImageWriteTask::EnsureWritableFile()
 		return true;
 	}
 	// We can't write to the file
+	else if (bOverwriteFile)
+	{
+		UE_LOG(LogImageWriteQueue, Error, TEXT("Failed to write image to '%s'. Should have overwritten the file, but we failed to delete it."), *Filename);
+		return false;
+	}
 	else
 	{
-		UE_LOG(LogImageWriteQueue, Error, TEXT("Failed to write image to '%s'. Should Overwrite: %d - If we should have overwritten the file, we failed to delete the file. If we shouldn't have overwritten the file the file already exists so we can't replace it."), *Filename, bOverwriteFile);
+		UE_LOG(LogImageWriteQueue, Error, TEXT("Failed to write image to '%s'. Shouldn't overwrite the file and it already exists so we can't replace it."), *Filename);
 		return false;
 	}
 }
