Replace use of TAsyncAlphaWrite with a simple call to SetAlphaOpaque
Remove format switches from client code so support of all formats is automatic
Use existing ImageCore SetAlphaOpaque
template header functions are now unused

#rnx

diff --git a/Engine/Plugins/MovieScene/MovieRenderPipeline/Source/MovieRenderPipelineRenderPasses/Private/MoviePipelineImageSequenceOutput.cpp b/Engine/Plugins/MovieScene/MovieRenderPipeline/Source/MovieRenderPipelineRenderPasses/Private/MoviePipelineImageSequenceOutput.cpp
--- a/Engine/Plugins/MovieScene/MovieRenderPipeline/Source/MovieRenderPipelineRenderPasses/Private/MoviePipelineImageSequenceOutput.cpp
+++ b/Engine/Plugins/MovieScene/MovieRenderPipeline/Source/MovieRenderPipelineRenderPasses/Private/MoviePipelineImageSequenceOutput.cpp
@@ -258,18 +258,7 @@ void UMoviePipelineImageSequenceOutputBase::OnReceiveImageDataImpl(FMoviePipelin
 		// no good without alpha, and we already did logic above to ensure it got turned into a filetype that could write alpha.
 		if (!IsAlphaAllowed() && !Payload->bRequireTransparentOutput)
 		{
-			switch (QuantizedPixelType)
-			{
-			case EImagePixelType::Color:
-				TileImageTask->PixelPreProcessors.Add(TAsyncAlphaWrite<FColor>(255));
-				break;
-			case EImagePixelType::Float16:
-				TileImageTask->PixelPreProcessors.Add(TAsyncAlphaWrite<FFloat16Color>(1.0f));
-				break;
-			case EImagePixelType::Float32:
-				TileImageTask->PixelPreProcessors.Add(TAsyncAlphaWrite<FLinearColor>(1.0f));
-				break;
-			}
+			TileImageTask->AddPreProcessorToSetAlphaOpaque();
 		}
 
 
diff --git a/Engine/Source/Editor/UnrealEd/Private/EditorViewportClient.cpp b/Engine/Source/Editor/UnrealEd/Private/EditorViewportClient.cpp
--- a/Engine/Source/Editor/UnrealEd/Private/EditorViewportClient.cpp
+++ b/Engine/Source/Editor/UnrealEd/Private/EditorViewportClient.cpp
@@ -5841,7 +5841,7 @@ void FEditorViewportClient::TakeScreenshot(FViewport* InViewport, bool bInValida
 		ImageTask->PixelData = MoveTemp(PixelData);
 
 		// Ensure the alpha channel is full alpha (this happens on the background thread)
-		ImageTask->PixelPreProcessors.Add(TAsyncAlphaWrite<FColor>(255));
+		ImageTask->AddPreProcessorToSetAlphaOpaque();
 	}
 	else
 	{
@@ -5862,7 +5862,7 @@ void FEditorViewportClient::TakeScreenshot(FViewport* InViewport, bool bInValida
 		ImageTask->PixelData = MoveTemp(PixelData);
 
 		// Ensure the alpha channel is full alpha (this happens on the background thread)
-		ImageTask->PixelPreProcessors.Add(TAsyncAlphaWrite<FLinearColor>(1.0f));
+		ImageTask->AddPreProcessorToSetAlphaOpaque();
 	}
 
 	// Create screenshot folder if not already present.
@@ -6041,7 +6041,7 @@ bool RequestSaveScreenshot(bool bWriteAlpha, TArray<FColorType>& Bitmap, FIntPoi
 		// Set full alpha on the bitmap
 		if (!bWriteAlpha)
 		{
-			ImageTask->PixelPreProcessors.Add(TAsyncAlphaWrite<FColorType>(OpaqueAlphaValue));
+			ImageTask->AddPreProcessorToSetAlphaOpaque();
 		}
 
 		HighResScreenshotConfig.PopulateImageTaskParams(*ImageTask);
diff --git a/Engine/Source/Runtime/ImageCore/Private/ImageCore.cpp b/Engine/Source/Runtime/ImageCore/Private/ImageCore.cpp
--- a/Engine/Source/Runtime/ImageCore/Private/ImageCore.cpp
+++ b/Engine/Source/Runtime/ImageCore/Private/ImageCore.cpp
@@ -1595,6 +1595,8 @@ bool FImageCore::DetectAlphaChannel(const FImageView & InImage)
 
 void FImageCore::SetAlphaOpaque(const FImageView & InImage)
 {
+	// this is not currently parallel but easily could be
+
 	int64 NumPixels = (int64)InImage.SizeX * InImage.SizeY * InImage.NumSlices;
 
 	if (InImage.Format == ERawImageFormat::BGRA8)
diff --git a/Engine/Source/Runtime/ImageWriteQueue/Private/ImageWriteTask.cpp b/Engine/Source/Runtime/ImageWriteQueue/Private/ImageWriteTask.cpp
--- a/Engine/Source/Runtime/ImageWriteQueue/Private/ImageWriteTask.cpp
+++ b/Engine/Source/Runtime/ImageWriteQueue/Private/ImageWriteTask.cpp
@@ -142,3 +142,15 @@ bool FImageWriteTask::EnsureWritableFile()
 	}
 }
 
+void FImageWriteTask::AddPreProcessorToSetAlphaOpaque()
+{
+	struct PreProcessorToSetAlphaOpaque
+	{
+		void operator()(FImagePixelData* InPixelData)
+		{
+			InPixelData->SetAlphaOpaque();
+		}
+	};
+
+	PixelPreProcessors.Emplace( PreProcessorToSetAlphaOpaque() );
+}
diff --git a/Engine/Source/Runtime/ImageWriteQueue/Public/ImagePixelData.h b/Engine/Source/Runtime/ImageWriteQueue/Public/ImagePixelData.h
--- a/Engine/Source/Runtime/ImageWriteQueue/Public/ImagePixelData.h
+++ b/Engine/Source/Runtime/ImageWriteQueue/Public/ImagePixelData.h
@@ -124,6 +124,14 @@ struct FImagePixelData
 		return Ret;
 	}
 
+	/**
+	* Change the alpha channel to opaque
+	*/
+	void SetAlphaOpaque()
+	{
+		FImageCore::SetAlphaOpaque( GetImageView() );
+	}
+
 	/**
 	 * Check that this data is the size it should be
 	 */
diff --git a/Engine/Source/Runtime/ImageWriteQueue/Public/ImageWriteTask.h b/Engine/Source/Runtime/ImageWriteQueue/Public/ImageWriteTask.h
--- a/Engine/Source/Runtime/ImageWriteQueue/Public/ImageWriteTask.h
+++ b/Engine/Source/Runtime/ImageWriteQueue/Public/ImageWriteTask.h
@@ -81,6 +81,8 @@ public:
 	IMAGEWRITEQUEUE_API virtual bool RunTask() override final;
 	IMAGEWRITEQUEUE_API virtual void OnAbandoned() override final;
 
+	IMAGEWRITEQUEUE_API void AddPreProcessorToSetAlphaOpaque();
+
 private:
 
 	/**
@@ -116,11 +118,13 @@ private:
 
 /**
  * A pixel preprocessor for use with FImageWriteTask::PixelPreProcessor that overwrites the alpha channel with a fixed value as part of the threaded work
+ *
+ * DEPRECATED.  Prefer AddPreProcessorToSetAlphaOpaque.
  */
 template<typename PixelType> struct TAsyncAlphaWrite;
 
 template<>
-struct TAsyncAlphaWrite<FColor>
+struct TAsyncAlphaWrite<FColor> // prefer AddPreProcessorToSetAlphaOpaque
 {
 	uint8 Alpha;
 	TAsyncAlphaWrite(uint8 InAlpha) : Alpha(InAlpha) {}
@@ -138,7 +142,7 @@ struct TAsyncAlphaWrite<FColor>
 };
 
 template<>
-struct TAsyncAlphaWrite<FFloat16Color>
+struct TAsyncAlphaWrite<FFloat16Color> // prefer AddPreProcessorToSetAlphaOpaque
 {
 	FFloat16 Alpha;
 	TAsyncAlphaWrite(float InAlpha) : Alpha(InAlpha) {}
@@ -156,7 +160,7 @@ struct TAsyncAlphaWrite<FFloat16Color>
 };
 
 template<>
-struct TAsyncAlphaWrite<FLinearColor>
+struct TAsyncAlphaWrite<FLinearColor> // prefer AddPreProcessorToSetAlphaOpaque
 {
 	float Alpha;
 	TAsyncAlphaWrite(float InAlpha) : Alpha(InAlpha) {}
@@ -175,6 +179,9 @@ struct TAsyncAlphaWrite<FLinearColor>
 
 /**
  * A pixel preprocessor for use with FImageWriteTask::PixelPreProcessor that inverts the alpha channel as part of the threaded work
+ *
+ * DEPRECATED.  This is not used that I can see; if it is, make something like AddPreProcessorToSetAlphaOpaque where the implementation
+ *   is in a C file and handles all formats, not a template in a header.
  */
 template<typename PixelType> struct TAsyncAlphaInvert;
 
@@ -203,7 +210,7 @@ struct TAsyncAlphaInvert<FFloat16Color>
 		TImagePixelData<FFloat16Color>* Float16ColorData = static_cast<TImagePixelData<FFloat16Color>*>(PixelData);
 		for (FFloat16Color& Pixel : Float16ColorData->Pixels)
 		{
-			Pixel.A = FFloat16(1.f) - Pixel.A;
+			Pixel.A = FFloat16(1.f - Pixel.A.GetFloat());
 		}
 	}
 };
diff --git a/Engine/Source/Runtime/MediaUtils/Private/MediaRecorder.cpp b/Engine/Source/Runtime/MediaUtils/Private/MediaRecorder.cpp
--- a/Engine/Source/Runtime/MediaUtils/Private/MediaRecorder.cpp
+++ b/Engine/Source/Runtime/MediaUtils/Private/MediaRecorder.cpp
@@ -274,22 +274,7 @@ void FMediaRecorder::TickRecording()
 				{
 					check(bUseFMediaImagePixelData);
 
-					if ( ImageTask->PixelData->GetType() == EImagePixelType::Float16 )
-					{
-						ImageTask->PixelPreProcessors.Add(TAsyncAlphaWrite<FFloat16Color>(1.f));
-					}
-					else if ( ImageTask->PixelData->GetType() == EImagePixelType::Float32 )
-					{
-						ImageTask->PixelPreProcessors.Add(TAsyncAlphaWrite<FLinearColor>(1.f));
-					}
-					else if ( ImageTask->PixelData->GetType() == EImagePixelType::Color )
-					{
-						ImageTask->PixelPreProcessors.Add(TAsyncAlphaWrite<FColor>(255));
-					}
-					else
-					{
-						check(false);
-					}
+					ImageTask->AddPreProcessorToSetAlphaOpaque();
 				}
 			}
 			
diff --git a/Engine/Source/Runtime/MovieSceneCapture/Private/ImageSequenceProtocol.cpp b/Engine/Source/Runtime/MovieSceneCapture/Private/ImageSequenceProtocol.cpp
--- a/Engine/Source/Runtime/MovieSceneCapture/Private/ImageSequenceProtocol.cpp
+++ b/Engine/Source/Runtime/MovieSceneCapture/Private/ImageSequenceProtocol.cpp
@@ -131,7 +131,7 @@ void UImageSequenceProtocol::ProcessFrame(FCapturedFrameData Frame)
 	if (Format == EImageFormat::PNG)
 	{
 		// Always write full alpha for PNGs
-		ImageTask->PixelPreProcessors.Add(TAsyncAlphaWrite<FColor>(255));
+		ImageTask->AddPreProcessorToSetAlphaOpaque();
 	}
 
 	switch (Format)
diff --git a/Engine/Source/Runtime/MovieSceneCapture/Private/UserDefinedCaptureProtocol.cpp b/Engine/Source/Runtime/MovieSceneCapture/Private/UserDefinedCaptureProtocol.cpp
--- a/Engine/Source/Runtime/MovieSceneCapture/Private/UserDefinedCaptureProtocol.cpp
+++ b/Engine/Source/Runtime/MovieSceneCapture/Private/UserDefinedCaptureProtocol.cpp
@@ -431,7 +431,7 @@ void UUserDefinedImageCaptureProtocol::WriteImageToDisk(const FCapturedPixels& P
 	// If the pixels are FColors, and this is the final pixels buffer, and we're writing PNG, always write out full alpha
 	if (PixelData.ImageData->GetType() == EImagePixelType::Color && ImageTask->Format == EImageFormat::PNG && StreamID.Identifiers.OrderIndependentCompareEqual(FinalPixelsID.Identifiers))
 	{
-		ImageTask->PixelPreProcessors.Add(TAsyncAlphaWrite<FColor>(255));
+		ImageTask->AddPreProcessorToSetAlphaOpaque();
 	}
 
 	if (Format == EDesiredImageFormat::EXR)
diff --git a/Engine/Source/Runtime/Renderer/Private/PostProcess/PostProcessVisualizeBuffer.cpp b/Engine/Source/Runtime/Renderer/Private/PostProcess/PostProcessVisualizeBuffer.cpp
--- a/Engine/Source/Runtime/Renderer/Private/PostProcess/PostProcessVisualizeBuffer.cpp
+++ b/Engine/Source/Runtime/Renderer/Private/PostProcess/PostProcessVisualizeBuffer.cpp
@@ -258,8 +258,8 @@ void AddDumpToFilePass(FRDGBuilder& GraphBuilder, FScreenPassTexture Input, cons
 
 		if (ImageTask->PixelData->GetType() == EImagePixelType::Color)
 		{
-			// Always write full alpha
-			ImageTask->PixelPreProcessors.Add(TAsyncAlphaWrite<FColor>(255));
+			// Always write opaque alpha
+			ImageTask->AddPreProcessorToSetAlphaOpaque();
 
 			// ImageTask->PixelData should be sRGB
 			//  it will gamma correct automatically if written to EXR
@@ -267,19 +267,7 @@ void AddDumpToFilePass(FRDGBuilder& GraphBuilder, FScreenPassTexture Input, cons
 		else if(ImageTask->Format == EImageFormat::PNG)
 		{
 			// PNGs can't have 0 alpha or RGB data is destroyed.
-
-			if ( ImageTask->PixelData->GetType() == EImagePixelType::Float32 )
-			{
-				ImageTask->PixelPreProcessors.Add(TAsyncAlphaWrite<FLinearColor>(1.0f));
-			}
-			else if ( ImageTask->PixelData->GetType() == EImagePixelType::Float16 )
-			{
-				ImageTask->PixelPreProcessors.Add(TAsyncAlphaWrite<FFloat16Color>(1.0f));
-			}
-			else
-			{
-				check(false);
-			}
+			ImageTask->AddPreProcessorToSetAlphaOpaque();
 		}
 
 		HighResScreenshotConfig.ImageWriteQueue->Enqueue(MoveTemp(ImageTask));
