If we have  GThreadPool, always use that, otherwise we can end up with an off by one error when comparing number of threads vs number of cores
As number of threads is generally now cores - 1 (or cores - 2)

#Synced-CL 51322414

diff --git a/Engine/Source/Runtime/ImageWriteQueue/Private/ImageWriteQueue.cpp b/Engine/Source/Runtime/ImageWriteQueue/Private/ImageWriteQueue.cpp
--- a/Engine/Source/Runtime/ImageWriteQueue/Private/ImageWriteQueue.cpp
+++ b/Engine/Source/Runtime/ImageWriteQueue/Private/ImageWriteQueue.cpp
@@ -285,7 +285,7 @@ void FImageWriteQueue::RecreateThreadPool()
 			bOwnedThreadPool = false;
 			ThreadPool = GIOThreadPool;
 		}
-		else if (GThreadPool && GThreadPool->GetNumThreads() >= MaxConcurrency)
+		else if (GThreadPool)
 		{
 			// Use a simple wrapper to limit concurrency and reuse threads we already have
 			bOwnedThreadPool = true;
